Fixes for double/float conversion warnings

diff --git a/Engine/Plugins/Animation/BlendSpaceMotionAnalysis/Source/BlendSpaceMotionAnalysis/BlendSpaceMotionAnalysis.Build.cs b/Engine/Plugins/Animation/BlendSpaceMotionAnalysis/Source/BlendSpaceMotionAnalysis/BlendSpaceMotionAnalysis.Build.cs
--- a/Engine/Plugins/Animation/BlendSpaceMotionAnalysis/Source/BlendSpaceMotionAnalysis/BlendSpaceMotionAnalysis.Build.cs
+++ b/Engine/Plugins/Animation/BlendSpaceMotionAnalysis/Source/BlendSpaceMotionAnalysis/BlendSpaceMotionAnalysis.Build.cs
@@ -8,6 +8,8 @@ namespace UnrealBuildTool.Rules
 	{
 		public BlendSpaceMotionAnalysis(ReadOnlyTargetRules Target) : base(Target)
 		{
+			UnsafeTypeCastWarningLevel = WarningLevel.Warning;
+
 			PublicDependencyModuleNames.AddRange(
 				new string[]
 				{
diff --git a/Engine/Plugins/Animation/BlendSpaceMotionAnalysis/Source/BlendSpaceMotionAnalysis/Private/LocomotionAnalysis.cpp b/Engine/Plugins/Animation/BlendSpaceMotionAnalysis/Source/BlendSpaceMotionAnalysis/Private/LocomotionAnalysis.cpp
--- a/Engine/Plugins/Animation/BlendSpaceMotionAnalysis/Source/BlendSpaceMotionAnalysis/Private/LocomotionAnalysis.cpp
+++ b/Engine/Plugins/Animation/BlendSpaceMotionAnalysis/Source/BlendSpaceMotionAnalysis/Private/LocomotionAnalysis.cpp
@@ -87,14 +87,14 @@ static bool CalculateLocomotionVelocity(
 	BlendSpaceAnalysis::GetFrameDirs(FrameFacingDir, FrameUpDir, FrameRightDir, FrameTM, AnalysisProperties);
 
 	// The frame time delta
-	float DeltaTime = Animation.GetPlayLength() / NumSampledKeys;
+	double DeltaTime = Animation.GetPlayLength() / double(NumSampledKeys);
 
 	// First step is to figure out the approximate direction. Note that the average velocity will be zero (assuming a
 	// complete cycle) - but if we apply a weight that is based on the height, then we can bias it towards the foot that
 	// is on the ground.
-	float MinHeight = FLT_MAX;
-	float MaxHeight = -FLT_MAX;
-	FVector AveragePos(0.0f);
+	double MinHeight = DBL_MAX;
+	double MaxHeight = -DBL_MAX;
+	FVector AveragePos(0);
 	TArray<FVector> Positions;
 	Positions.SetNum(NumSampledKeys);
 	for (int32 Key = 0; Key != NumSampledKeys; ++Key)
@@ -103,7 +103,7 @@ static bool CalculateLocomotionVelocity(
 		FTransform TM = BoneOffset * BoneTM;
 		FVector Pos = TM.GetTranslation();
 		Positions[Key] = Pos;
-		float Height = Pos | FrameUpDir;
+		double Height = Pos | FrameUpDir;
 		MinHeight = FMath::Min(MinHeight, Height);
 		MaxHeight = FMath::Max(MaxHeight, Height);
 		AveragePos += Pos;
@@ -120,15 +120,15 @@ static bool CalculateLocomotionVelocity(
 	{
 		int32 PrevKey = (Key + NumSampledKeys - 1) % NumSampledKeys;
 		int32 NextKey = (Key +  1) % NumSampledKeys;
-		Velocities[Key] = (Positions[NextKey] - Positions[PrevKey]) / (2.0f * DeltaTime);
+		Velocities[Key] = (Positions[NextKey] - Positions[PrevKey]) / (2.0 * DeltaTime);
 	}
 
 	FVector BiasedFootVel(0);
-	float TotalWeight = 0.0f;
+	double TotalWeight = 0;
 	for (int32 Key = 0 ; Key != NumSampledKeys ; ++Key)
 	{
-		float Height = Positions[Key] | FrameUpDir;
-		float Weight = 1.0f - (Height - MinHeight) / (MaxHeight - MinHeight);
+		double Height = Positions[Key] | FrameUpDir;
+		double Weight = 1.0 - (Height - MinHeight) / (MaxHeight - MinHeight);
 		BiasedFootVel += Velocities[Key] * Weight;
 		TotalWeight += Weight;
 	}
@@ -148,7 +148,7 @@ static bool CalculateLocomotionVelocity(
 	Mask.SetNum(NumSampledKeys);
 	for (int32 Key = 0 ; Key != NumSampledKeys ; ++Key)
 	{
-		Mask[Key] = (Velocities[Key] | ApproxLocoDir) >= 0.0f ? 0 : 1;
+		Mask[Key] = (Velocities[Key] | ApproxLocoDir) >= 0 ? 0 : 1;
 	}
 
 	int32 StartKey = -1;
@@ -206,14 +206,14 @@ static bool CalculateLocomotionVelocity(
 	int32 Threshold = FMath::Max(MaxMask / 2, 1);
 	int32 Num = 0;
 	FVector AverageFootVel(0);
-	float BestSpeed = 0.0f;
+	double BestSpeed = 0;
 	int32 BestSpeedKey = 0;
 	for (int32 K = AZeroKey ; K != AZeroKey + NumSampledKeys ; ++K)
 	{
 		int32 Key = K % NumSampledKeys;
 		if (Mask[Key] >= Threshold)
 		{
-			float Speed = Velocities[Key].Size();
+			double Speed = Velocities[Key].Size();
 #ifdef ANALYSIS_VERBOSE_LOG
 			UE_LOG(LogAnimation, Log, TEXT("Candidate %d Mask %d vel = %f %f %f, speed = %f"), 
 				   Key, Mask[Key], Velocities[Key].X, Velocities[Key].Y, Velocities[Key].Z, Speed);
@@ -226,7 +226,7 @@ static bool CalculateLocomotionVelocity(
 		}
 		else
 		{
-			if (BestSpeed > 0.0f)
+			if (BestSpeed > 0)
 			{
 #ifdef ANALYSIS_VERBOSE_LOG
 				UE_LOG(LogAnimation, Log, TEXT("Picked Candidate %d vel = %f %f %f, speed = %f"), 
@@ -234,12 +234,12 @@ static bool CalculateLocomotionVelocity(
 #endif
 				AverageFootVel += Velocities[BestSpeedKey];
 				++Num;
-				BestSpeed = 0.0f;
+				BestSpeed = 0;
 			}
 		}
 	}
 	// Make sure we didn't miss the last data point
-	if (BestSpeed > 0.0f)
+	if (BestSpeed > 0)
 	{
 #ifdef ANALYSIS_VERBOSE_LOG
 		UE_LOG(LogAnimation, Log, TEXT("Picked Candidate %d vel = %f %f %f, speed = %f"), 
@@ -247,13 +247,13 @@ static bool CalculateLocomotionVelocity(
 #endif
 		AverageFootVel += Velocities[BestSpeedKey];
 		++Num;
-		BestSpeed = 0.0f;
+		BestSpeed = 0;
 	}
 
 	AverageFootVel /= Num;
-	float FacingVel = -AverageFootVel | FrameFacingDir;
-	float RightVel = -AverageFootVel | FrameRightDir;
-	float UpVel = -AverageFootVel | FrameUpDir;
+	double FacingVel = -AverageFootVel | FrameFacingDir;
+	double RightVel = -AverageFootVel | FrameRightDir;
+	double UpVel = -AverageFootVel | FrameUpDir;
 
 	Result.Set(FacingVel, RightVel, UpVel);
 	Result *= Animation.RateScale * RateScale;
@@ -307,7 +307,7 @@ static bool CalculateLocomotionVelocity(
 //======================================================================================================================
 // Calculates the locomotion speed (magnitude)
 static bool CalculateLocomotionSpeed(
-	float&                               Result,
+	double&                              Result,
 	const UBlendSpace&                   BlendSpace,
 	const ULocomotionAnalysisProperties* AnalysisProperties,
 	const UAnimSequence&                 Animation,
@@ -325,7 +325,7 @@ static bool CalculateLocomotionSpeed(
 //======================================================================================================================
 // Calculates the locomotion direction (degrees)
 static bool CalculateLocomotionDirection(
-	float&                               Result,
+	double&                              Result,
 	const UBlendSpace&                   BlendSpace,
 	const ULocomotionAnalysisProperties* AnalysisProperties,
 	const UAnimSequence&                 Animation,
@@ -343,7 +343,7 @@ static bool CalculateLocomotionDirection(
 //======================================================================================================================
 // Calculates the locomotion speed in the character's facing direction
 static bool CalculateLocomotionForwardSpeed(
-	float&                               Result,
+	double&                              Result,
 	const UBlendSpace&                   BlendSpace,
 	const ULocomotionAnalysisProperties* AnalysisProperties,
 	const UAnimSequence&                 Animation,
@@ -361,7 +361,7 @@ static bool CalculateLocomotionForwardSpeed(
 //======================================================================================================================
 // Calculates the locomotion speed in the character's upwards direction
 static bool CalculateLocomotionUpwardSpeed(
-	float&                               Result,
+	double&                              Result,
 	const UBlendSpace&                   BlendSpace,
 	const ULocomotionAnalysisProperties* AnalysisProperties,
 	const UAnimSequence&                 Animation,
@@ -379,7 +379,7 @@ static bool CalculateLocomotionUpwardSpeed(
 //======================================================================================================================
 // Calculates the locomotion speed in the character's right direction
 static bool CalculateLocomotionRightwardSpeed(
-	float&                               Result,
+	double&                              Result,
 	const UBlendSpace&                   BlendSpace,
 	const ULocomotionAnalysisProperties* AnalysisProperties,
 	const UAnimSequence&                 Animation,
@@ -397,7 +397,7 @@ static bool CalculateLocomotionRightwardSpeed(
 //======================================================================================================================
 // Calculates the locomotion slope angle (degrees) going in the facing direction
 static bool CalculateLocomotionForwardSlope(
-	float&                               Result,
+	double&                              Result,
 	const UBlendSpace&                   BlendSpace,
 	const ULocomotionAnalysisProperties* AnalysisProperties,
 	const UAnimSequence&                 Animation,
@@ -406,7 +406,7 @@ static bool CalculateLocomotionForwardSlope(
 	FVector Movement;
 	if (CalculateLocomotionVelocity(Movement, BlendSpace, AnalysisProperties, Animation, RateScale))
 	{
-		if (Movement.X >= 0.0f)
+		if (Movement.X >= 0)
 		{
 			Result = FMath::RadiansToDegrees(FMath::Atan2(Movement.Z, Movement.X));
 		}
@@ -423,7 +423,7 @@ static bool CalculateLocomotionForwardSlope(
 //======================================================================================================================
 // Calculates the locomotion slope angle (degrees) going in the rightwards direction
 static bool CalculateLocomotionRightwardSlope(
-	float&                               Result,
+	double&                              Result,
 	const UBlendSpace&                   BlendSpace,
 	const ULocomotionAnalysisProperties* AnalysisProperties,
 	const UAnimSequence&                 Animation,
@@ -432,7 +432,7 @@ static bool CalculateLocomotionRightwardSlope(
 	FVector Movement;
 	if (CalculateLocomotionVelocity(Movement, BlendSpace, AnalysisProperties, Animation, RateScale))
 	{
-		if (Movement.Y > 0.0f)
+		if (Movement.Y > 0)
 		{
 			Result = FMath::RadiansToDegrees(FMath::Atan2(Movement.Z, Movement.Y));
 		}
@@ -446,8 +446,8 @@ static bool CalculateLocomotionRightwardSlope(
 }
 
 //======================================================================================================================
-bool CalculateLocomotion(
-	float&                               Result,
+static bool CalculateLocomotion(
+	double&                              Result,
 	const UBlendSpace&                   BlendSpace,
 	const ULocomotionAnalysisProperties* AnalysisProperties,
 	const UAnimSequence&                 Animation,
@@ -479,4 +479,20 @@ bool CalculateLocomotion(
 	}
 }
 
+//======================================================================================================================
+// Note that it is easier to do the calculations involving world-space positions in doubles, and
+// then cast Result to float here, than it is to be casting in all the functions above.
+bool CalculateLocomotion(
+	float&                               Result,
+	const UBlendSpace&                   BlendSpace,
+	const ULocomotionAnalysisProperties* AnalysisProperties,
+	const UAnimSequence&                 Animation,
+	const float                          RateScale)
+{
+	double ResultDouble = Result;
+	bool bResult = CalculateLocomotion(ResultDouble, BlendSpace, AnalysisProperties, Animation, RateScale);
+	Result = FloatCastChecked<float>(ResultDouble, UE::LWC::DefaultFloatPrecision);
+	return bResult;
+}
+
 #undef LOCTEXT_NAMESPACE
diff --git a/Engine/Plugins/Animation/BlendSpaceMotionAnalysis/Source/BlendSpaceMotionAnalysis/Private/RootMotionAnalysis.cpp b/Engine/Plugins/Animation/BlendSpaceMotionAnalysis/Source/BlendSpaceMotionAnalysis/Private/RootMotionAnalysis.cpp
--- a/Engine/Plugins/Animation/BlendSpaceMotionAnalysis/Source/BlendSpaceMotionAnalysis/Private/RootMotionAnalysis.cpp
+++ b/Engine/Plugins/Animation/BlendSpaceMotionAnalysis/Source/BlendSpaceMotionAnalysis/Private/RootMotionAnalysis.cpp
@@ -74,9 +74,9 @@ static bool CalculateRootMotionVelocity(
 		FVector FrameFacingDir = BlendSpaceAnalysis::GetAxisFromTM(FrameTM, AnalysisProperties->CharacterFacingAxis);
 		FVector FrameUpDir = BlendSpaceAnalysis::GetAxisFromTM(FrameTM, AnalysisProperties->CharacterUpAxis);
 		FVector FrameRightDir = FVector::CrossProduct(FrameUpDir, FrameFacingDir);
-		float ForwardVel = Velocity | FrameFacingDir;
-		float RightVel = Velocity | FrameRightDir;
-		float UpVel = Velocity | FrameUpDir;
+		double ForwardVel = Velocity | FrameFacingDir;
+		double RightVel = Velocity | FrameRightDir;
+		double UpVel = Velocity | FrameUpDir;
 		Result.Set(ForwardVel, RightVel, UpVel);
 		return true;
 	}
@@ -86,7 +86,7 @@ static bool CalculateRootMotionVelocity(
 //======================================================================================================================
 // Calculates the movement speed (magnitude) 
 static bool CalculateRootMotionSpeed(
-	float&                               Result,
+	double&                              Result,
 	const UBlendSpace&                   BlendSpace,
 	const URootMotionAnalysisProperties* AnalysisProperties,
 	const UAnimSequence&                 Animation,
@@ -104,7 +104,7 @@ static bool CalculateRootMotionSpeed(
 //======================================================================================================================
 // Calculates the movement speed (magnitude) 
 static bool CalculateRootMotionDirection(
-	float&                               Result,
+	double&                              Result,
 	const UBlendSpace&                   BlendSpace,
 	const URootMotionAnalysisProperties* AnalysisProperties,
 	const UAnimSequence&                 Animation,
@@ -122,7 +122,7 @@ static bool CalculateRootMotionDirection(
 //======================================================================================================================
 // Calculates the locomotion speed in the character's facing direction
 static bool CalculateRootMotionForwardSpeed(
-	float&                               Result,
+	double&                              Result,
 	const UBlendSpace&                   BlendSpace,
 	const URootMotionAnalysisProperties* AnalysisProperties,
 	const UAnimSequence&                 Animation,
@@ -140,7 +140,7 @@ static bool CalculateRootMotionForwardSpeed(
 //======================================================================================================================
 // Calculates the locomotion speed in the character's upwards direction
 static bool CalculateRootMotionUpwardSpeed(
-	float&                               Result,
+	double&                              Result,
 	const UBlendSpace&                   BlendSpace,
 	const URootMotionAnalysisProperties* AnalysisProperties,
 	const UAnimSequence&                 Animation,
@@ -158,7 +158,7 @@ static bool CalculateRootMotionUpwardSpeed(
 //======================================================================================================================
 // Calculates the locomotion speed in the character's right direction
 static bool CalculateRootMotionRightwardSpeed(
-	float&                               Result,
+	double&                              Result,
 	const UBlendSpace&                   BlendSpace,
 	const URootMotionAnalysisProperties* AnalysisProperties,
 	const UAnimSequence&                 Animation,
@@ -176,7 +176,7 @@ static bool CalculateRootMotionRightwardSpeed(
 //======================================================================================================================
 // Calculates the locomotion slope angle (degrees) going in the facing direction
 static bool CalculateRootMotionForwardSlope(
-	float&                               Result,
+	double&                              Result,
 	const UBlendSpace&                   BlendSpace,
 	const URootMotionAnalysisProperties* AnalysisProperties,
 	const UAnimSequence&                 Animation,
@@ -185,7 +185,7 @@ static bool CalculateRootMotionForwardSlope(
 	FVector Movement;
 	if (CalculateRootMotionVelocity(Movement, BlendSpace, AnalysisProperties, Animation, RateScale))
 	{
-		if (Movement.X >= 0.0f)
+		if (Movement.X >= 0)
 		{
 			Result = FMath::RadiansToDegrees(FMath::Atan2(Movement.Z, Movement.X));
 		}
@@ -202,7 +202,7 @@ static bool CalculateRootMotionForwardSlope(
 //======================================================================================================================
 // Calculates the locomotion slope angle (degrees) going in the rightwards direction
 static bool CalculateRootMotionRightwardSlope(
-	float&                               Result,
+	double&                              Result,
 	const UBlendSpace&                   BlendSpace,
 	const URootMotionAnalysisProperties* AnalysisProperties,
 	const UAnimSequence&                 Animation,
@@ -211,7 +211,7 @@ static bool CalculateRootMotionRightwardSlope(
 	FVector Movement;
 	if (CalculateRootMotionVelocity(Movement, BlendSpace, AnalysisProperties, Animation, RateScale))
 	{
-		if (Movement.Y > 0.0f)
+		if (Movement.Y > 0)
 		{
 			Result = FMath::RadiansToDegrees(FMath::Atan2(Movement.Z, Movement.Y));
 		}
@@ -226,9 +226,8 @@ static bool CalculateRootMotionRightwardSlope(
 
 
 //======================================================================================================================
-// Calculates the movement speed (magnitude) 
-bool CalculateRootMotion(
-	float&                               Result,
+static bool CalculateRootMotion(
+	double&                              Result,
 	const UBlendSpace&                   BlendSpace,
 	const URootMotionAnalysisProperties* AnalysisProperties,
 	const UAnimSequence&                 Animation,
@@ -268,4 +267,20 @@ bool CalculateRootMotion(
 	}
 }
 
+//======================================================================================================================
+// Note that it is easier to do the calculations involving world-space positions in doubles, and
+// then cast Result to float here, than it is to be casting in all the functions above.
+bool CalculateRootMotion(
+	float&                               Result,
+	const UBlendSpace&                   BlendSpace,
+	const URootMotionAnalysisProperties* AnalysisProperties,
+	const UAnimSequence&                 Animation,
+	const float                          RateScale)
+{
+	double DoubleResult = Result;
+	bool bResult = CalculateRootMotion(DoubleResult, BlendSpace, AnalysisProperties, Animation, RateScale);
+	Result = FloatCastChecked<float>(DoubleResult, UE::LWC::DefaultFloatPrecision);
+	return bResult;
+}
+
 #undef LOCTEXT_NAMESPACE
diff --git a/Engine/Plugins/Animation/BlendSpaceMotionAnalysis/Source/BlendSpaceMotionAnalysis/Public/LocomotionAnalysis.h b/Engine/Plugins/Animation/BlendSpaceMotionAnalysis/Source/BlendSpaceMotionAnalysis/Public/LocomotionAnalysis.h
--- a/Engine/Plugins/Animation/BlendSpaceMotionAnalysis/Source/BlendSpaceMotionAnalysis/Public/LocomotionAnalysis.h
+++ b/Engine/Plugins/Animation/BlendSpaceMotionAnalysis/Source/BlendSpaceMotionAnalysis/Public/LocomotionAnalysis.h
@@ -48,7 +48,10 @@ public:
 	EAnalysisLinearAxis CharacterUpAxis = EAnalysisLinearAxis::PlusZ;
 };
 
-//======================================================================================================================
+/**
+ * Calculates the locomotion speed from the animation (which may be playrate scaled), 
+ * according to the analysis properties.
+ */
 bool CalculateLocomotion(
 	float&                               Result,
 	const UBlendSpace&                   BlendSpace,
@@ -56,7 +59,6 @@ bool CalculateLocomotion(
 	const UAnimSequence&                 Animation,
 	const float                          RateScale);
 
-
 #if UE_ENABLE_INCLUDE_ORDER_DEPRECATED_IN_5_2
 #include "CoreMinimal.h"
 #include "Modules/ModuleInterface.h"
diff --git a/Engine/Plugins/Animation/BlendSpaceMotionAnalysis/Source/BlendSpaceMotionAnalysis/Public/RootMotionAnalysis.h b/Engine/Plugins/Animation/BlendSpaceMotionAnalysis/Source/BlendSpaceMotionAnalysis/Public/RootMotionAnalysis.h
--- a/Engine/Plugins/Animation/BlendSpaceMotionAnalysis/Source/BlendSpaceMotionAnalysis/Public/RootMotionAnalysis.h
+++ b/Engine/Plugins/Animation/BlendSpaceMotionAnalysis/Source/BlendSpaceMotionAnalysis/Public/RootMotionAnalysis.h
@@ -65,7 +65,10 @@ public:
 	float EndTimeFraction = 1.0f;
 };
 
-//======================================================================================================================
+/**
+ * Calculates the root motion movement speed from the animation (which may be playrate scaled), 
+ * according to the analysis properties.
+ */
 bool CalculateRootMotion(
 	float&                               Result,
 	const UBlendSpace&                   BlendSpace,
@@ -73,6 +76,7 @@ bool CalculateRootMotion(
 	const UAnimSequence&                 Animation,
 	const float                          RateScale);
 
+
 #if UE_ENABLE_INCLUDE_ORDER_DEPRECATED_IN_5_2
 #include "CoreMinimal.h"
 #include "Modules/ModuleInterface.h"
diff --git a/Engine/Plugins/Experimental/PhysicsControl/Source/PhysicsControl/PhysicsControl.Build.cs b/Engine/Plugins/Experimental/PhysicsControl/Source/PhysicsControl/PhysicsControl.Build.cs
--- a/Engine/Plugins/Experimental/PhysicsControl/Source/PhysicsControl/PhysicsControl.Build.cs
+++ b/Engine/Plugins/Experimental/PhysicsControl/Source/PhysicsControl/PhysicsControl.Build.cs
@@ -7,7 +7,8 @@ public class PhysicsControl : ModuleRules
 	public PhysicsControl(ReadOnlyTargetRules Target) : base(Target)
 	{
 		PCHUsage = ModuleRules.PCHUsageMode.UseExplicitOrSharedPCHs;
-		
+		UnsafeTypeCastWarningLevel = WarningLevel.Warning;
+
 		PublicDependencyModuleNames.AddRange(
 			new string[]
 			{
diff --git a/Engine/Plugins/Experimental/PhysicsControl/Source/PhysicsControl/Private/PhysicsControlComponent.cpp b/Engine/Plugins/Experimental/PhysicsControl/Source/PhysicsControl/Private/PhysicsControlComponent.cpp
--- a/Engine/Plugins/Experimental/PhysicsControl/Source/PhysicsControl/Private/PhysicsControlComponent.cpp
+++ b/Engine/Plugins/Experimental/PhysicsControl/Source/PhysicsControl/Private/PhysicsControlComponent.cpp
@@ -1024,7 +1024,7 @@ bool UPhysicsControlComponent::SetControlTargetPosition(
 	FPhysicsControlRecord* Record = Implementation->FindControlRecord(Name);
 	if (Record)
 	{
-		if (VelocityDeltaTime)
+		if (VelocityDeltaTime != 0)
 		{
 			Record->PhysicsControl.ControlTarget.TargetVelocity =
 				(Position - Record->PhysicsControl.ControlTarget.TargetPosition) / VelocityDeltaTime;
@@ -1055,7 +1055,7 @@ bool UPhysicsControlComponent::SetControlTargetOrientation(
 	FPhysicsControlRecord* Record = Implementation->FindControlRecord(Name);
 	if (Record)
 	{
-		if (AngularVelocityDeltaTime)
+		if (AngularVelocityDeltaTime != 0)
 		{
 			FQuat OldQ = Record->PhysicsControl.ControlTarget.TargetOrientation.Quaternion();
 			FQuat OrientationQ = Orientation.Quaternion();
@@ -1100,7 +1100,7 @@ bool UPhysicsControlComponent::SetControlTargetPoses(
 		FVector Position = OffsetTM.GetTranslation();
 		FQuat OrientationQ = OffsetTM.GetRotation();
 
-		if (VelocityDeltaTime)
+		if (VelocityDeltaTime != 0)
 		{
 			FQuat OldQ = Record->PhysicsControl.ControlTarget.TargetOrientation.Quaternion();
 			OldQ.EnforceShortestArcWith(OrientationQ);
@@ -1952,11 +1952,11 @@ void UPhysicsControlComponent::DebugDrawControl(
 			DrawCoordinateSystem(
 				PDI, WorldTargetTM.GetTranslation(), WorldTargetTM.Rotator(),
 				VisualizationSizeScale, SDPG_Foreground, 4.0f * GizmoWidthScale);
-			if (VelocityPredictionTime)
+			if (VelocityPredictionTime != 0)
 			{
 				DrawCoordinateSystem(
 					PDI, PredictedTargetTM.GetTranslation(), PredictedTargetTM.Rotator(),
-					VisualizationSizeScale * 0.5, SDPG_Foreground, 4.0f * GizmoWidthScale);
+					VisualizationSizeScale * 0.5f, SDPG_Foreground, 4.0f * GizmoWidthScale);
 			}
 		}
 		else
@@ -1967,16 +1967,16 @@ void UPhysicsControlComponent::DebugDrawControl(
 			DrawWireSphere(
 				PDI, WorldTargetTM, TargetColor, 
 				VisualizationSizeScale, 8, SDPG_Foreground, 3.0f * GizmoWidthScale);
-			if (VelocityPredictionTime)
+			if (VelocityPredictionTime != 0)
 			{
 				DrawWireSphere(
 					PDI, PredictedTargetTM, TargetColor, 
-					VisualizationSizeScale * 0.5, 8, SDPG_Foreground, 3.0f * GizmoWidthScale);
+					VisualizationSizeScale * 0.5f, 8, SDPG_Foreground, 3.0f * GizmoWidthScale);
 			}
 		}
 
 
-		if (VelocityPredictionTime)
+		if (VelocityPredictionTime != 0)
 		{
 			PDI->DrawLine(
 				WorldTargetTM.GetTranslation(), 
diff --git a/Engine/Plugins/Experimental/PhysicsControl/Source/PhysicsControl/Private/PhysicsControlComponentHelpers.cpp b/Engine/Plugins/Experimental/PhysicsControl/Source/PhysicsControl/Private/PhysicsControlComponentHelpers.cpp
--- a/Engine/Plugins/Experimental/PhysicsControl/Source/PhysicsControl/Private/PhysicsControlComponentHelpers.cpp
+++ b/Engine/Plugins/Experimental/PhysicsControl/Source/PhysicsControl/Private/PhysicsControlComponentHelpers.cpp
@@ -12,19 +12,6 @@ namespace UE
 namespace PhysicsControlComponent
 {
 
-//======================================================================================================================
-void ConvertStrengthToSpringParams(
-	double& OutSpring, double& OutDamping, 
-	double InStrength, double InDampingRatio, double InExtraDamping)
-{
-	double AngularFrequency = InStrength * UE_DOUBLE_TWO_PI;
-	double Stiffness = AngularFrequency * AngularFrequency;
-
-	OutSpring = Stiffness;
-	OutDamping = 2.0 * InDampingRatio * AngularFrequency;
-	OutDamping += InExtraDamping;
-}
-
 //======================================================================================================================
 void ConvertStrengthToSpringParams(
 	FVector& OutSpring, FVector& OutDamping, 
@@ -35,36 +22,6 @@ void ConvertStrengthToSpringParams(
 	ConvertStrengthToSpringParams(OutSpring.Z, OutDamping.Z, InStrength.Z, InDampingRatio, InExtraDamping.Z);
 }
 
-//======================================================================================================================
-void ConvertSpringToStrengthParams(
-	float& OutStrength, float& OutDampingRatio, float& OutExtraDamping,
-	double InSpring, double InDamping)
-{
-	// Simple calculation to get the strength
-	double AngularFrequency = FMath::Sqrt(InSpring);
-	OutStrength = AngularFrequency / UE_DOUBLE_TWO_PI;
-
-	// For damping, try to put as much into the damping ratio as possible, up to a max DR of 1. Then
-	// the rest goes into extra damping.
-	OutDampingRatio = 1.0;
-	double ImpliedDamping = 2.0 * OutDampingRatio * AngularFrequency;
-
-	if (ImpliedDamping < InDamping)
-	{
-		OutExtraDamping = InDamping - ImpliedDamping;
-	}
-	else if (AngularFrequency > 0)
-	{
-		OutExtraDamping = 0;
-		OutDampingRatio = InDamping / (2.0 * AngularFrequency);
-	}
-	else
-	{
-		OutDampingRatio = 1.0;
-		OutExtraDamping = InDamping;
-	}
-}
-
 inline double GetLinearDriveStiffness(const FLinearDriveConstraint& InDrive)
 {
 	return (
@@ -137,7 +94,8 @@ void ConvertConstraintProfileToControlData(
 		GetLinearDriveStiffness(InProfileProperties.LinearDrive), 
 		GetLinearDriveDamping(InProfileProperties.LinearDrive));
 
-	OutControlData.MaxForce = GetLinearDriveMaxForce(InProfileProperties.LinearDrive);
+	OutControlData.MaxForce = FloatCastChecked<float>(
+		GetLinearDriveMaxForce(InProfileProperties.LinearDrive), UE::LWC::DefaultFloatPrecision);
 
 	ConvertSpringToStrengthParams(
 		OutControlData.AngularStrength,
@@ -146,7 +104,8 @@ void ConvertConstraintProfileToControlData(
 		GetAngularDriveStiffness(InProfileProperties.AngularDrive),
 		GetAngularDriveDamping(InProfileProperties.AngularDrive));
 
-	OutControlData.MaxTorque = GetAngularDriveMaxTorque(InProfileProperties.AngularDrive);
+	OutControlData.MaxTorque = FloatCastChecked<float>(
+		GetAngularDriveMaxTorque(InProfileProperties.AngularDrive), UE::LWC::DefaultFloatPrecision);
 }
 
 //======================================================================================================================
diff --git a/Engine/Plugins/Experimental/PhysicsControl/Source/PhysicsControl/Private/PhysicsControlComponentHelpers.h b/Engine/Plugins/Experimental/PhysicsControl/Source/PhysicsControl/Private/PhysicsControlComponentHelpers.h
--- a/Engine/Plugins/Experimental/PhysicsControl/Source/PhysicsControl/Private/PhysicsControlComponentHelpers.h
+++ b/Engine/Plugins/Experimental/PhysicsControl/Source/PhysicsControl/Private/PhysicsControlComponentHelpers.h
@@ -3,6 +3,7 @@
 #pragma once
 
 #include "UObject/NameTypes.h"
+#include "Math/Vector.h"
 
 struct FBodyInstance;
 struct FPhysicsControlData;
@@ -16,13 +17,6 @@ namespace UE
 namespace PhysicsControlComponent
 {
 
-/** 
- * Converts strength/damping ratio/extra damping into spring stiffness/damping.
- */
-void ConvertStrengthToSpringParams(
-	double& OutSpring, double& OutDamping, 
-	double InStrength, double InDampingRatio, double InExtraDamping);
-
 /** 
  * Converts strength/damping ratio/extra damping into spring stiffness/damping 
  */
@@ -30,14 +24,6 @@ void ConvertStrengthToSpringParams(
 	FVector& OutSpring, FVector& OutDamping, 
 	const FVector& InStrength, float InDampingRatio, const FVector& InExtraDamping);
 
-/**
- * Converts spring/damping values into strength/damping ratio/extra damping. This tries to get
- * as much damping into the damping ratio term as possible, without letting it go above 1.
- */
-void ConvertSpringToStrengthParams(
-	double& OutStrength, double& OutDampingRatio, double& OutExtraDamping,
-	double InSpring, double InDamping);
-
 /**
  * Converts the drive settings from the constraint profile into the control data strength/damping etc. 
  * Approximations will be made if (a) the linear drive has different values for the x/y/z axes, or (b)
@@ -60,5 +46,55 @@ FBodyInstance* GetBodyInstance(UMeshComponent* MeshComponent, const FName BoneNa
  */
 FName GetPhysicalParentBone(USkeletalMeshComponent* SkeletalMeshComponent, FName BoneName);
 
+/**
+ * Converts strength/damping ratio/extra damping into spring stiffness/damping.
+ */
+template<typename TOut>
+void ConvertStrengthToSpringParams(
+	TOut& OutSpring, TOut& OutDamping,
+	double InStrength, double InDampingRatio, double InExtraDamping)
+{
+	TOut AngularFrequency = FloatCastChecked<TOut>(InStrength * UE_DOUBLE_TWO_PI, UE::LWC::DefaultFloatPrecision);
+	TOut Stiffness = AngularFrequency * AngularFrequency;
+
+	OutSpring = Stiffness;
+	OutDamping = FloatCastChecked<TOut>(InExtraDamping + 2 * InDampingRatio * AngularFrequency, UE::LWC::DefaultFloatPrecision);
+}
+
+
+/**
+ * Converts spring/damping values into strength/damping ratio/extra damping. This tries to get
+ * as much damping into the damping ratio term as possible, without letting it go above 1.
+ */
+template<typename TOut>
+static void ConvertSpringToStrengthParams(
+	TOut& OutStrength, TOut& OutDampingRatio, TOut& OutExtraDamping,
+	const double InSpring, const double InDamping)
+{
+	// Simple calculation to get the strength
+	TOut AngularFrequency = FloatCastChecked<TOut>(FMath::Sqrt(InSpring), UE::LWC::DefaultFloatPrecision);
+	OutStrength = AngularFrequency / UE_TWO_PI;
+
+	// For damping, try to put as much into the damping ratio as possible, up to a max DR of 1. Then
+	// the rest goes into extra damping.
+	OutDampingRatio = 1;
+	TOut ImpliedDamping = FloatCastChecked<TOut>(2.0 * OutDampingRatio * AngularFrequency, UE::LWC::DefaultFloatPrecision);
+
+	if (ImpliedDamping < InDamping)
+	{
+		OutExtraDamping = FloatCastChecked<TOut>(InDamping, UE::LWC::DefaultFloatPrecision) - ImpliedDamping;
+	}
+	else if (AngularFrequency > 0)
+	{
+		OutExtraDamping = 0;
+		OutDampingRatio = FloatCastChecked<TOut>(InDamping, UE::LWC::DefaultFloatPrecision) / (2 * AngularFrequency);
+	}
+	else
+	{
+		OutDampingRatio = 1;
+		OutExtraDamping = FloatCastChecked<TOut>(InDamping, UE::LWC::DefaultFloatPrecision);
+	}
+}
+
 } // namespace PhysicsControlComponent
 } // namespace UE
diff --git a/Engine/Plugins/Experimental/PhysicsControl/Source/PhysicsControl/Private/PhysicsControlComponentImpl.cpp b/Engine/Plugins/Experimental/PhysicsControl/Source/PhysicsControl/Private/PhysicsControlComponentImpl.cpp
--- a/Engine/Plugins/Experimental/PhysicsControl/Source/PhysicsControl/Private/PhysicsControlComponentImpl.cpp
+++ b/Engine/Plugins/Experimental/PhysicsControl/Source/PhysicsControl/Private/PhysicsControlComponentImpl.cpp
@@ -427,7 +427,7 @@ void FPhysicsControlComponentImpl::CalculateControlTargetData(
 
 				if (bCalculateVelocity)
 				{
-					if (Record.PhysicsControl.ControlSettings.SkeletalAnimationVelocityMultiplier)
+					if (Record.PhysicsControl.ControlSettings.SkeletalAnimationVelocityMultiplier != 0)
 					{
 						// Offset of the control point from the target child bone TM, in world space.
 						FVector WorldControlPointOffset =
@@ -470,7 +470,7 @@ void FPhysicsControlComponentImpl::CalculateControlTargetData(
 					OutTargetVelocity = ChildBoneTM.GetRotation() * OutTargetVelocity;
 					OutTargetAngularVelocity = ChildBoneTM.GetRotation() * OutTargetAngularVelocity;
 
-					if (Record.PhysicsControl.ControlSettings.SkeletalAnimationVelocityMultiplier)
+					if (Record.PhysicsControl.ControlSettings.SkeletalAnimationVelocityMultiplier != 0)
 					{
 						FVector WorldControlPointOffset =
 							ChildBoneTM.GetRotation() * Record.PhysicsControl.ControlSettings.ControlPoint;
@@ -530,9 +530,9 @@ bool FPhysicsControlComponentImpl::ApplyControlStrengths(
 	const FPhysicsControlData& Data = Record.PhysicsControl.ControlData;
 	const FPhysicsControlMultiplier& Multiplier = Record.PhysicsControl.ControlMultiplier;
 
-	double AngularSpring;
-	double AngularDamping;
-	double MaxTorque = Data.MaxTorque * Multiplier.MaxTorqueMultiplier;
+	float AngularSpring;
+	float AngularDamping;
+	float MaxTorque = Data.MaxTorque * Multiplier.MaxTorqueMultiplier;
 
 	FVector LinearSpring;
 	FVector LinearDamping;
@@ -549,25 +549,25 @@ bool FPhysicsControlComponentImpl::ApplyControlStrengths(
 		Data.LinearDampingRatio,
 		Data.LinearExtraDamping * Multiplier.LinearExtraDampingMultiplier);
 
-	if (Multiplier.MaxTorqueMultiplier <= 0.0)
+	if (Multiplier.MaxTorqueMultiplier <= 0)
 	{
-		AngularSpring = 0.0;
-		AngularDamping = 0.0;
+		AngularSpring = 0;
+		AngularDamping = 0;
 	}
-	if (Multiplier.MaxForceMultiplier.X <= 0.0)
+	if (Multiplier.MaxForceMultiplier.X <= 0)
 	{
-		LinearSpring.X = 0.0;
-		LinearDamping.X = 0.0;
+		LinearSpring.X = 0;
+		LinearDamping.X = 0;
 	}
-	if (Multiplier.MaxForceMultiplier.Y <= 0.0)
+	if (Multiplier.MaxForceMultiplier.Y <= 0)
 	{
-		LinearSpring.Y = 0.0;
-		LinearDamping.Y = 0.0;
+		LinearSpring.Y = 0;
+		LinearDamping.Y = 0;
 	}
-	if (Multiplier.MaxForceMultiplier.Z <= 0.0)
+	if (Multiplier.MaxForceMultiplier.Z <= 0)
 	{
-		LinearSpring.Z = 0.0;
-		LinearDamping.Z = 0.0;
+		LinearSpring.Z = 0;
+		LinearDamping.Z = 0;
 	}
 
 	ConstraintInstance->SetAngularDriveParams(AngularSpring, AngularDamping, MaxTorque);
@@ -579,7 +579,7 @@ bool FPhysicsControlComponentImpl::ApplyControlStrengths(
 		FMath::Max(UE_SMALL_NUMBER, MaxForce.Y),
 		FMath::Max(UE_SMALL_NUMBER, MaxForce.Z));
 	double Test = TestAngular + TestLinear.GetMax();
-	return Test > 0.0;
+	return Test > 0;
 }
 
 //======================================================================================================================
diff --git a/Engine/Plugins/Experimental/PhysicsControl/Source/PhysicsControlVisualizer/PhysicsControlVisualizer.Build.cs b/Engine/Plugins/Experimental/PhysicsControl/Source/PhysicsControlVisualizer/PhysicsControlVisualizer.Build.cs
--- a/Engine/Plugins/Experimental/PhysicsControl/Source/PhysicsControlVisualizer/PhysicsControlVisualizer.Build.cs
+++ b/Engine/Plugins/Experimental/PhysicsControl/Source/PhysicsControlVisualizer/PhysicsControlVisualizer.Build.cs
@@ -7,7 +7,8 @@ public class PhysicsControlVisualizer : ModuleRules
 	public PhysicsControlVisualizer(ReadOnlyTargetRules Target) : base(Target)
 	{
 		PCHUsage = ModuleRules.PCHUsageMode.UseExplicitOrSharedPCHs;
-		
+		UnsafeTypeCastWarningLevel = WarningLevel.Warning;
+
 		PublicIncludePaths.AddRange(
 			new string[] {
 				// ... add public include paths required here ...
diff --git a/Engine/Source/Editor/Persona/Private/PersonaBlendSpaceAnalysis.cpp b/Engine/Source/Editor/Persona/Private/PersonaBlendSpaceAnalysis.cpp
--- a/Engine/Source/Editor/Persona/Private/PersonaBlendSpaceAnalysis.cpp
+++ b/Engine/Source/Editor/Persona/Private/PersonaBlendSpaceAnalysis.cpp
@@ -302,12 +302,12 @@ FVector BlendSpaceAnalysis::CalculateSampleValue(const UBlendSpace& BlendSpace,
 		const UAnalysisProperties* AnalysisProperties = BlendSpace.AnalysisProperties[Index].Get();
 		for (const IBlendSpaceAnalysisFeature* Feature : ModularFeatures)
 		{
-			float FloatValue = (float) AdjustedPosition[Index];
+			float NewPosition = FloatCastChecked<float>(AdjustedPosition[Index], UE::LWC::DefaultFloatPrecision);
 			bAnalyzed[Index] = Feature->CalculateSampleValue(
-				FloatValue, BlendSpace, AnalysisProperties, Animation, RateScale);
+				NewPosition, BlendSpace, AnalysisProperties, Animation, RateScale);
 			if (bAnalyzed[Index])
 			{
-				AdjustedPosition[Index] = (FVector::FReal) FloatValue;
+				AdjustedPosition[Index] = NewPosition;
 				break;
 			}
 		}
diff --git a/Engine/Source/Editor/Persona/Private/SAnimationBlendSpace.cpp b/Engine/Source/Editor/Persona/Private/SAnimationBlendSpace.cpp
--- a/Engine/Source/Editor/Persona/Private/SAnimationBlendSpace.cpp
+++ b/Engine/Source/Editor/Persona/Private/SAnimationBlendSpace.cpp
@@ -175,7 +175,7 @@ int32 SBlendSpaceEditor::OnSampleAdded(UAnimSequence* Animation, const FVector&
 
 	if(BlendSpace->IsAsset() && bRunAnalysis)
 	{
-		AdjustedValue = BlendSpaceAnalysis::CalculateSampleValue(*BlendSpace, *Animation, 1.0f, Value, bAnalyzed);
+		AdjustedValue = BlendSpaceAnalysis::CalculateSampleValue(*BlendSpace, *Animation, 1.0, Value, bAnalyzed);
 	}
 		
 	int32 NewSampleIndex = -1;
diff --git a/Engine/Source/Editor/Persona/Public/BlendSpaceAnalysis.h b/Engine/Source/Editor/Persona/Public/BlendSpaceAnalysis.h
--- a/Engine/Source/Editor/Persona/Public/BlendSpaceAnalysis.h
+++ b/Engine/Source/Editor/Persona/Public/BlendSpaceAnalysis.h
@@ -251,7 +251,7 @@ void GetFrameDirs(
 */
 template<typename FunctionType, typename T>
 static bool CalculateComponentSampleValue(
-	float&                     Result,
+	double&                    Result,
 	const FunctionType&        Fn,
 	const UBlendSpace&         BlendSpace, 
 	const T*                   AnalysisProperties, 
@@ -262,12 +262,31 @@ static bool CalculateComponentSampleValue(
 	int32 ComponentIndex = (int32) AnalysisProperties->FunctionAxis;
 	if (Fn(Value, BlendSpace, AnalysisProperties, Animation, RateScale))
 	{
-		Result = static_cast<float>(Value[ComponentIndex]);
+		Result = Value[ComponentIndex];
 		return true;
 	}
 	return false;
 }
 
+//======================================================================================================================
+/**
+* Helper to extract the component from the FVector functions
+*/
+template<typename FunctionType, typename T>
+static bool CalculateComponentSampleValue(
+	float&               Result,
+	const FunctionType&  Fn,
+	const UBlendSpace&   BlendSpace,
+	const T*             AnalysisProperties,
+	const UAnimSequence& Animation,
+	const float          RateScale)
+{
+	double DoubleResult = Result;
+	bool bResult = CalculateComponentSampleValue(DoubleResult, Fn, BlendSpace, AnalysisProperties, Animation, RateScale);
+	Result = FloatCastChecked<float>(DoubleResult, UE::LWC::DefaultFloatPrecision);
+	return bResult;
+}
+
 //======================================================================================================================
 template <typename T>
 static bool CalculatePosition(
@@ -291,9 +310,9 @@ static bool CalculatePosition(
 	}
 
 	int32 FirstKey = FMath::Clamp(
-		(int32) (NumSampledKeys * AnalysisProperties->StartTimeFraction), 0, NumSampledKeys);
+		(int32) (float(NumSampledKeys) * AnalysisProperties->StartTimeFraction), 0, NumSampledKeys);
 	int32 LastKey = FMath::Clamp(
-		(int32) (NumSampledKeys * AnalysisProperties->EndTimeFraction), FirstKey, NumSampledKeys);
+		(int32) (float(NumSampledKeys) * AnalysisProperties->EndTimeFraction), FirstKey, NumSampledKeys);
 
 	FTransform FrameTM;
 	bool bNeedToUpdateFrameTM = true;
@@ -335,9 +354,9 @@ static bool CalculateDeltaPosition(
 	}
 
 	int32 FirstKey = FMath::Clamp(
-		(int32) (NumSampledKeys * AnalysisProperties->StartTimeFraction), 0, NumSampledKeys);
+		(int32) (float(NumSampledKeys) * AnalysisProperties->StartTimeFraction), 0, NumSampledKeys);
 	int32 LastKey = FMath::Clamp(
-		(int32) (NumSampledKeys * AnalysisProperties->EndTimeFraction), FirstKey, NumSampledKeys);
+		(int32) (float(NumSampledKeys) * AnalysisProperties->EndTimeFraction), FirstKey, NumSampledKeys);
 
 	FTransform FrameTM;
 	bool bNeedToUpdateFrameTM = true;
@@ -384,12 +403,12 @@ static bool CalculateVelocity(
 		return false;
 	}
 
-	float DeltaTime = Animation.GetPlayLength() / NumSampledKeys;
+	double DeltaTime = Animation.GetPlayLength() / double(NumSampledKeys);
 
 	int32 FirstKey = FMath::Clamp(
-		(int32) (NumSampledKeys * AnalysisProperties->StartTimeFraction), 0, NumSampledKeys);
+		(int32) (float(NumSampledKeys) * AnalysisProperties->StartTimeFraction), 0, NumSampledKeys);
 	int32 LastKey = FMath::Clamp(
-		(int32) (NumSampledKeys * AnalysisProperties->EndTimeFraction), FirstKey, NumSampledKeys-1);
+		(int32) (float(NumSampledKeys) * AnalysisProperties->EndTimeFraction), FirstKey, NumSampledKeys-1);
 
 	// First and Last key are for averaging. However, the finite differencing always goes from one frame to the next
 	int32 NumKeys = FMath::Max(1 + LastKey - FirstKey, 1);
@@ -500,9 +519,9 @@ static bool CalculateOrientation(
 	}
 
 	int32 FirstKey = FMath::Clamp(
-		(int32) (NumSampledKeys * AnalysisProperties->StartTimeFraction), 0, NumSampledKeys);
+		(int32) (float(NumSampledKeys) * AnalysisProperties->StartTimeFraction), 0, NumSampledKeys);
 	int32 LastKey = FMath::Clamp(
-		(int32) (NumSampledKeys * AnalysisProperties->EndTimeFraction), FirstKey, NumSampledKeys);
+		(int32) (float(NumSampledKeys) * AnalysisProperties->EndTimeFraction), FirstKey, NumSampledKeys);
 
 	FTransform FrameTM;
 	bool bNeedToUpdateFrameTM = true;
@@ -555,9 +574,9 @@ static bool CalculateDeltaOrientation(
 	}
 
 	int32 FirstKey = FMath::Clamp(
-		(int32) (NumSampledKeys * AnalysisProperties->StartTimeFraction), 0, NumSampledKeys);
+		(int32) (float(NumSampledKeys) * AnalysisProperties->StartTimeFraction), 0, NumSampledKeys);
 	int32 LastKey = FMath::Clamp(
-		(int32) (NumSampledKeys * AnalysisProperties->EndTimeFraction), FirstKey, NumSampledKeys);
+		(int32) (float(NumSampledKeys) * AnalysisProperties->EndTimeFraction), FirstKey, NumSampledKeys);
 
 	FTransform FrameTM;
 	bool bNeedToUpdateFrameTM = true;
@@ -610,12 +629,12 @@ static bool CalculateAngularVelocity(
 		return false;
 	}
 
-	float DeltaTime = Animation.GetPlayLength() / NumSampledKeys;
+	double DeltaTime = Animation.GetPlayLength() / double(NumSampledKeys);
 
 	int32 FirstKey = FMath::Clamp(
-		(int32) (NumSampledKeys * AnalysisProperties->StartTimeFraction), 0, NumSampledKeys);
+		(int32) (float(NumSampledKeys) * AnalysisProperties->StartTimeFraction), 0, NumSampledKeys);
 	int32 LastKey = FMath::Clamp(
-		(int32) (NumSampledKeys * AnalysisProperties->EndTimeFraction), FirstKey, NumSampledKeys-1);
+		(int32) (float(NumSampledKeys) * AnalysisProperties->EndTimeFraction), FirstKey, NumSampledKeys-1);
 
 	// First and Last key are for averaging. However, the finite differencing always goes from one frame to the next
 	int32 NumKeys = FMath::Max(1 + LastKey - FirstKey, 1);
@@ -646,7 +665,7 @@ static bool CalculateAngularVelocity(
 
 		FQuat Rotation = RelativeQuat2 * RelativeQuat1.Inverse();
 		FVector Axis;
-		float Angle;
+		double Angle;
 		Rotation.ToAxisAndAngle(Axis, Angle);
 		FVector AngularVelocity = FMath::RadiansToDegrees(Axis * (Angle / DeltaTime));
 #ifdef ANALYSIS_VERBOSE_LOG
@@ -688,12 +707,12 @@ static bool CalculateOrientationRate(
 		return false;
 	}
 
-	float DeltaTime = Animation.GetPlayLength() / NumSampledKeys;
+	double DeltaTime = Animation.GetPlayLength() / double(NumSampledKeys);
 
 	int32 FirstKey = FMath::Clamp(
-		(int32) (NumSampledKeys * AnalysisProperties->StartTimeFraction), 0, NumSampledKeys);
+		(int32) (float(NumSampledKeys) * AnalysisProperties->StartTimeFraction), 0, NumSampledKeys);
 	int32 LastKey = FMath::Clamp(
-		(int32) (NumSampledKeys * AnalysisProperties->EndTimeFraction), FirstKey, NumSampledKeys-1);
+		(int32) (float(NumSampledKeys) * AnalysisProperties->EndTimeFraction), FirstKey, NumSampledKeys-1);
 
 	// First and Last key are for averaging. However, the finite differencing always goes from one frame to the next
 	int32 NumKeys = FMath::Max(1 + LastKey - FirstKey, 1);
