Updated several modules to set UnsafeTypeCastWarningLevel to WarningLevel.Error and not WarningLevel.Warning.  The use of WarningLevel.Warning makes separate shared PCHs.

diff --git a/Engine/Plugins/Animation/BlendSpaceMotionAnalysis/Source/BlendSpaceMotionAnalysis/BlendSpaceMotionAnalysis.Build.cs b/Engine/Plugins/Animation/BlendSpaceMotionAnalysis/Source/BlendSpaceMotionAnalysis/BlendSpaceMotionAnalysis.Build.cs
--- a/Engine/Plugins/Animation/BlendSpaceMotionAnalysis/Source/BlendSpaceMotionAnalysis/BlendSpaceMotionAnalysis.Build.cs
+++ b/Engine/Plugins/Animation/BlendSpaceMotionAnalysis/Source/BlendSpaceMotionAnalysis/BlendSpaceMotionAnalysis.Build.cs
@@ -8,7 +8,7 @@ namespace UnrealBuildTool.Rules
 	{
 		public BlendSpaceMotionAnalysis(ReadOnlyTargetRules Target) : base(Target)
 		{
-			UnsafeTypeCastWarningLevel = WarningLevel.Warning;
+			UnsafeTypeCastWarningLevel = WarningLevel.Error;
 
 			PublicDependencyModuleNames.AddRange(
 				new string[]
diff --git a/Engine/Plugins/Experimental/PhysicsControl/Source/PhysicsControl/PhysicsControl.Build.cs b/Engine/Plugins/Experimental/PhysicsControl/Source/PhysicsControl/PhysicsControl.Build.cs
--- a/Engine/Plugins/Experimental/PhysicsControl/Source/PhysicsControl/PhysicsControl.Build.cs
+++ b/Engine/Plugins/Experimental/PhysicsControl/Source/PhysicsControl/PhysicsControl.Build.cs
@@ -7,7 +7,7 @@ public class PhysicsControl : ModuleRules
 	public PhysicsControl(ReadOnlyTargetRules Target) : base(Target)
 	{
 		PCHUsage = ModuleRules.PCHUsageMode.UseExplicitOrSharedPCHs;
-		UnsafeTypeCastWarningLevel = WarningLevel.Warning;
+		UnsafeTypeCastWarningLevel = WarningLevel.Error;
 
 		PublicDependencyModuleNames.AddRange(
 			new string[]
diff --git a/Engine/Plugins/Experimental/PhysicsControl/Source/PhysicsControlEditor/PhysicsControlEditor.Build.cs b/Engine/Plugins/Experimental/PhysicsControl/Source/PhysicsControlEditor/PhysicsControlEditor.Build.cs
--- a/Engine/Plugins/Experimental/PhysicsControl/Source/PhysicsControlEditor/PhysicsControlEditor.Build.cs
+++ b/Engine/Plugins/Experimental/PhysicsControl/Source/PhysicsControlEditor/PhysicsControlEditor.Build.cs
@@ -6,7 +6,7 @@ public class PhysicsControlEditor : ModuleRules
 	public PhysicsControlEditor(ReadOnlyTargetRules Target) : base(Target)
 	{
 		PCHUsage = ModuleRules.PCHUsageMode.UseExplicitOrSharedPCHs;
-		UnsafeTypeCastWarningLevel = WarningLevel.Warning;
+		UnsafeTypeCastWarningLevel = WarningLevel.Error;
 
 		PublicIncludePaths.AddRange(
 			new string[] {
diff --git a/Engine/Plugins/Runtime/MassGameplay/Source/MassGameplayTestSuite/MassGameplayTestSuite.Build.cs b/Engine/Plugins/Runtime/MassGameplay/Source/MassGameplayTestSuite/MassGameplayTestSuite.Build.cs
--- a/Engine/Plugins/Runtime/MassGameplay/Source/MassGameplayTestSuite/MassGameplayTestSuite.Build.cs
+++ b/Engine/Plugins/Runtime/MassGameplay/Source/MassGameplayTestSuite/MassGameplayTestSuite.Build.cs
@@ -8,9 +8,6 @@ namespace UnrealBuildTool.Rules
 	{
 		public MassGameplayTestSuite(ReadOnlyTargetRules Target) : base(Target)
 		{
-			OptimizeCode = CodeOptimization.Never;
-			bUseUnity = false;
-
 			PCHUsage = ModuleRules.PCHUsageMode.UseExplicitOrSharedPCHs;
 
 			PublicIncludePaths.AddRange(
diff --git a/Engine/Source/Developer/CookOnTheFlyNetServer/CookOnTheFlyNetServer.Build.cs b/Engine/Source/Developer/CookOnTheFlyNetServer/CookOnTheFlyNetServer.Build.cs
--- a/Engine/Source/Developer/CookOnTheFlyNetServer/CookOnTheFlyNetServer.Build.cs
+++ b/Engine/Source/Developer/CookOnTheFlyNetServer/CookOnTheFlyNetServer.Build.cs
@@ -5,8 +5,8 @@ using UnrealBuildTool;
 public class CookOnTheFlyNetServer : ModuleRules
 {
     public CookOnTheFlyNetServer(ReadOnlyTargetRules Target) : base(Target)
-    {
-		UnsafeTypeCastWarningLevel = WarningLevel.Warning;
+	{
+		UnsafeTypeCastWarningLevel = WarningLevel.Error;
 		PrivateDependencyModuleNames.AddRange(
                 new string[]
                 {
diff --git a/Engine/Source/Developer/PakFileUtilities/PakFileUtilities.Build.cs b/Engine/Source/Developer/PakFileUtilities/PakFileUtilities.Build.cs
--- a/Engine/Source/Developer/PakFileUtilities/PakFileUtilities.Build.cs
+++ b/Engine/Source/Developer/PakFileUtilities/PakFileUtilities.Build.cs
@@ -6,7 +6,7 @@ public class PakFileUtilities : ModuleRules
 {
 	public PakFileUtilities(ReadOnlyTargetRules Target) : base(Target)
 	{
-		UnsafeTypeCastWarningLevel = WarningLevel.Warning;
+		UnsafeTypeCastWarningLevel = WarningLevel.Error;
 		PrivateDependencyModuleNames.AddRange(new string[] {
 			"Core",
 			"PakFile",
diff --git a/Engine/Source/Runtime/CookOnTheFly/CookOnTheFly.Build.cs b/Engine/Source/Runtime/CookOnTheFly/CookOnTheFly.Build.cs
--- a/Engine/Source/Runtime/CookOnTheFly/CookOnTheFly.Build.cs
+++ b/Engine/Source/Runtime/CookOnTheFly/CookOnTheFly.Build.cs
@@ -6,7 +6,7 @@ public class CookOnTheFly : ModuleRules
 {
 	public CookOnTheFly(ReadOnlyTargetRules Target) : base(Target)
 	{
-		UnsafeTypeCastWarningLevel = WarningLevel.Warning;
+		UnsafeTypeCastWarningLevel = WarningLevel.Error;
 		PublicDependencyModuleNames.AddRange(
             new string[] {
                 "Core",
diff --git a/Engine/Source/Runtime/PakFile/PakFile.Build.cs b/Engine/Source/Runtime/PakFile/PakFile.Build.cs
--- a/Engine/Source/Runtime/PakFile/PakFile.Build.cs
+++ b/Engine/Source/Runtime/PakFile/PakFile.Build.cs
@@ -6,7 +6,7 @@ public class PakFile : ModuleRules
 {
 	public PakFile(ReadOnlyTargetRules Target) : base(Target)
 	{
-		UnsafeTypeCastWarningLevel = WarningLevel.Warning;
+		UnsafeTypeCastWarningLevel = WarningLevel.Error;
 		PrivateDependencyModuleNames.Add("Core");
 		PrivateDependencyModuleNames.Add("CoreUObject");
 		PublicDependencyModuleNames.Add("RSA");
